Update LiveLink binaries for Mobu and Maya.
When Maya close a subject, tell the provider so it wouldn't resend it to UE4 once it connect back.

diff --git a/Engine/Source/Programs/MayaLiveLinkPlugin/BuildMayaPlugin.xml b/Engine/Source/Programs/MayaLiveLinkPlugin/BuildMayaPlugin.xml
--- a/Engine/Source/Programs/MayaLiveLinkPlugin/BuildMayaPlugin.xml
+++ b/Engine/Source/Programs/MayaLiveLinkPlugin/BuildMayaPlugin.xml
@@ -3,6 +3,7 @@
   <Property Name="LocalSourceDir" Value="$(RootDir)\Engine\Source\Programs\MayaLiveLinkPlugin"/>
   <Property Name="LocalStagingDir" Value="$(LocalSourceDir)\Staging"/>
   <Property Name="LocalBinaryDir" Value="$(RootDir)\Engine\Binaries\Win64"/>
+   <Property Name="LocalExtraDir" Value="$(RootDir)\Engine\Extras\MayaLiveLink"/>
   
   <Agent Name="MayaLiveLinkPlugin" Type="Win64">
     <Node Name="Compile UnrealHeaderTool Win64">
@@ -20,14 +21,14 @@
 	<Node Name="Compile Maya 2018 Win64" Requires="Compile UnrealHeaderTool Win64">
       <Compile Target="MayaLiveLinkPlugin2018" Platform="Win64" Configuration="Development" />
     </Node>
-	
-		<Node Name="Compile Maya 2019 Win64" Requires="Compile UnrealHeaderTool Win64">
+
+	<Node Name="Compile Maya 2019 Win64" Requires="Compile UnrealHeaderTool Win64">
       <Compile Target="MayaLiveLinkPlugin2019" Platform="Win64" Configuration="Development" />
     </Node>
 	
 	<Node Name="Stage Maya Plugin Module" Requires="Compile Maya 2016 Win64;Compile Maya 2017 Win64;Compile Maya 2018 Win64;Compile Maya 2019 Win64">
 		<Copy From="$(LocalSourceDir)\ReadMe.txt" To="$(LocalStagingDir)\ReadMe.txt" />
-		
+	
 		<Copy From="$(LocalBinaryDir)\MayaLiveLinkPlugin2016.mll" To="$(LocalStagingDir)\MayaLiveLink2016\LiveLink\plug-ins\MayaLiveLinkPlugin2016.mll" />
 		<Copy From="$(LocalSourceDir)\MayaLiveLinkUI.py" To="$(LocalStagingDir)\MayaLiveLink2016\LiveLink\plug-ins\MayaLiveLinkUI.py" />
 		<Copy From="$(LocalSourceDir)\LiveLink.mod" To="$(LocalStagingDir)\MayaLiveLink2016\LiveLink.mod" />
@@ -43,8 +44,10 @@
 		<Copy From="$(LocalBinaryDir)\MayaLiveLinkPlugin2019.mll" To="$(LocalStagingDir)\MayaLiveLink2019\LiveLink\plug-ins\MayaLiveLinkPlugin2019.mll" />
 		<Copy From="$(LocalSourceDir)\MayaLiveLinkUI.py" To="$(LocalStagingDir)\MayaLiveLink2019\LiveLink\plug-ins\MayaLiveLinkUI.py" />
 		<Copy From="$(LocalSourceDir)\LiveLink.mod" To="$(LocalStagingDir)\MayaLiveLink2019\LiveLink.mod" />
-		
+
 		<Zip FromDir="$(LocalStagingDir)" ZipFile="$(LocalStagingDir)\LiveLinkMaya.zip"/>
+		
+		<Copy From="$(LocalStagingDir)\LiveLinkMaya.zip" To="$(LocalExtraDir)\LiveLinkMaya.zip" />
 	</Node>
   </Agent>
 </BuildGraph>
diff --git a/Engine/Source/Programs/MayaLiveLinkPlugin/MayaLiveLinkPlugin.cpp b/Engine/Source/Programs/MayaLiveLinkPlugin/MayaLiveLinkPlugin.cpp
--- a/Engine/Source/Programs/MayaLiveLinkPlugin/MayaLiveLinkPlugin.cpp
+++ b/Engine/Source/Programs/MayaLiveLinkPlugin/MayaLiveLinkPlugin.cpp
@@ -244,6 +244,14 @@ struct FLiveLinkStreamedJointHeirarchySubject : IStreamedEntity
 		, RootDagPath(InRootPath)
 	{}
 
+	~FLiveLinkStreamedJointHeirarchySubject()
+	{
+		if (LiveLinkProvider.IsValid())
+		{
+			LiveLinkProvider->RemoveSubject(SubjectName);
+		}
+	}
+
 	virtual bool ShouldDisplayInUI() const { return true; }
 	virtual MString GetDisplayText() const { return MString("Character: ") + MString(*SubjectName.ToString()) + " ( " + RootDagPath.fullPathName() + " )"; }
 
@@ -386,6 +394,14 @@ struct FLiveLinkBaseCameraStreamedSubject : public IStreamedEntity
 public:
 	FLiveLinkBaseCameraStreamedSubject(FName InSubjectName) : SubjectName(InSubjectName) {}
 
+	~FLiveLinkBaseCameraStreamedSubject()
+	{
+		if (LiveLinkProvider.IsValid())
+		{
+			LiveLinkProvider->RemoveSubject(SubjectName);
+		}
+	}
+
 	virtual bool ValidateSubject() const { return true; }
 
 	virtual void RebuildSubjectData()
@@ -498,6 +514,14 @@ public:
 		, RootDagPath(InRootPath)
 	{}
 
+	~FLiveLinkStreamedPropSubject()
+	{
+		if (LiveLinkProvider.IsValid())
+		{
+			LiveLinkProvider->RemoveSubject(SubjectName);
+		}
+	}
+
 	virtual bool ShouldDisplayInUI() const { return true; }
 	virtual MString GetDisplayText() const { return MString("Prop: ") + MString(*SubjectName.ToString()) + " ( " + RootDagPath.fullPathName() + " )"; }
 
@@ -527,6 +551,7 @@ public:
 private:
 	FName SubjectName;
 	MDagPath RootDagPath;
+	bool bHasBeenRebuild;
 
 	static TArray<FName> PropBoneNames;
 	static TArray<int32> PropBoneParents;
@@ -599,10 +624,17 @@ public:
 
 	void RemoveSubject(MString SubjectToRemove)
 	{
-		TArray<MString> Entries;
-		GetSubjectEntries(Entries);
-		int32 Index = Entries.IndexOfByKey(SubjectToRemove);
-		Subjects.RemoveAt(Index);
+		for (int32 Index = Subjects.Num() - 1; Index >= 0; --Index)
+		{
+			if (Subjects[Index]->ShouldDisplayInUI())
+			{
+				if (Subjects[Index]->GetDisplayText() == SubjectToRemove)
+				{
+					Subjects.RemoveAt(Index);
+					break;
+				}
+			}
+		}
 	}
 
 	void Reset()
@@ -927,6 +959,10 @@ DLLEXPORT MStatus initializePlugin(MObject MayaPluginObject)
 	{
 		GEngineLoop.PreInit(TEXT("MayaLiveLinkPlugin -Messaging"));
 		ProcessNewlyLoadedUObjects();
+
+		// ensure target platform manager is referenced early as it must be created on the main thread
+		GetTargetPlatformManager();
+
 		// Tell the module manager is may now process newly-loaded UObjects when new C++ modules are loaded
 		FModuleManager::Get().StartProcessingNewlyLoadedObjects();
 
@@ -1023,4 +1059,4 @@ DLLEXPORT MStatus uninitializePlugin(MObject MayaPluginObject)
 
 	const MStatus MayaStatusResult = MS::kSuccess;
 	return MayaStatusResult;
-}
\ No newline at end of file
+}
diff --git a/Engine/Source/Programs/MayaLiveLinkPlugin/MayaLiveLinkPlugin2016.Build.cs b/Engine/Source/Programs/MayaLiveLinkPlugin/MayaLiveLinkPlugin2016.Build.cs
--- a/Engine/Source/Programs/MayaLiveLinkPlugin/MayaLiveLinkPlugin2016.Build.cs
+++ b/Engine/Source/Programs/MayaLiveLinkPlugin/MayaLiveLinkPlugin2016.Build.cs
@@ -8,19 +8,19 @@ public abstract class MayaLiveLinkPluginBase : ModuleRules
 	public MayaLiveLinkPluginBase(ReadOnlyTargetRules Target) : base(Target)
 	{
 		// For LaunchEngineLoop.cpp include.  You shouldn't need to add anything else to this line.
-		PrivateIncludePaths.AddRange( new string[] { "Runtime/Launch/Public", "Runtime/Launch/Private" }  );
+		PrivateIncludePaths.AddRange(new string[] { "Runtime/Launch/Public", "Runtime/Launch/Private" });
 
 		// Unreal dependency modules
-		PrivateDependencyModuleNames.AddRange( new string[] 
+		PrivateDependencyModuleNames.AddRange(new string[]
 		{
 			"Core",
-            "CoreUObject",
+			"CoreUObject",
 			"ApplicationCore",
 			"Projects",
-            "UdpMessaging",
-            "LiveLinkInterface",
-            "LiveLinkMessageBusFramework",
-		} );
+			"UdpMessaging",
+			"LiveLinkInterface",
+			"LiveLinkMessageBusFramework",
+		});
 
 
 		//
@@ -29,37 +29,43 @@ public abstract class MayaLiveLinkPluginBase : ModuleRules
 
 		{
 			string MayaVersionString = GetMayaVersion();
-			string MayaInstallFolder = @"C:\Program Files\Autodesk\Maya" + MayaVersionString;
+			string MayaInstallFolder = GetMayaInstallFolderPath();
 
 			// Make sure this version of Maya is actually installed
-			if( Directory.Exists( MayaInstallFolder ) )
+			if (Directory.Exists(MayaInstallFolder))
 			{
-                //throw new BuildException( "Couldn't find Autodesk Maya " + MayaVersionString + " in folder '" + MayaInstallFolder + "'.  This version of Maya must be installed for us to find the Maya SDK files." );
+				// These are required for Maya headers to compile properly as a DLL
+				PublicDefinitions.Add("NT_PLUGIN=1");
+				PublicDefinitions.Add("REQUIRE_IOSTREAM=1");
 
-                // These are required for Maya headers to compile properly as a DLL
-                PublicDefinitions.Add("NT_PLUGIN=1");
-                PublicDefinitions.Add("REQUIRE_IOSTREAM=1");
+				string IncludePath = GetMayaIncludePath();
+				PrivateIncludePaths.Add(IncludePath);
 
-                PrivateIncludePaths.Add(Path.Combine(MayaInstallFolder, "include"));
+				if (Target.Platform == UnrealTargetPlatform.Win64)  // @todo: Support other platforms?
+				{
+					PublicLibraryPaths.Add(GetMayaLibraryPath());
 
-                if (Target.Platform == UnrealTargetPlatform.Win64)  // @todo: Support other platforms?
-                {
-                    PublicLibraryPaths.Add(Path.Combine(MayaInstallFolder, "lib"));
-
-                    // Maya libraries we're depending on
-                    PublicAdditionalLibraries.AddRange(new string[]
-                        {
-                            "Foundation.lib",
-                            "OpenMaya.lib",
-                            "OpenMayaAnim.lib",
-                            "OpenMayaUI.lib"}
-                    );
-                }
-            }
+					// Maya libraries we're depending on
+					PublicAdditionalLibraries.AddRange(new string[]
+						{
+							"Foundation.lib",
+							"OpenMaya.lib",
+							"OpenMayaAnim.lib",
+							"OpenMayaUI.lib"}
+					);
+				}
+			}
+			else
+			{
+				// throw new BuildException("Couldn't find Autodesk Maya " + MayaVersionString + " in folder '" + MayaInstallFolder + "'.  This version of Maya must be installed for us to find the Maya SDK files.");
+			}
 		}
 	}
-	
+
 	public abstract string GetMayaVersion();
+	public virtual string GetMayaInstallFolderPath() { return @"C:\Program Files\Autodesk\Maya" + GetMayaVersion(); }
+	public virtual string GetMayaIncludePath() { return Path.Combine(GetMayaInstallFolderPath(), "include"); }
+	public virtual string GetMayaLibraryPath() { return Path.Combine(GetMayaInstallFolderPath(), "lib"); }
 }
 
 public class MayaLiveLinkPlugin2016 : MayaLiveLinkPluginBase
@@ -67,6 +73,6 @@ public class MayaLiveLinkPlugin2016 : MayaLiveLinkPluginBase
 	public MayaLiveLinkPlugin2016(ReadOnlyTargetRules Target) : base(Target)
 	{
 	}
-	
+
 	public override string GetMayaVersion() { return "2016"; }
 }
diff --git a/Engine/Source/Programs/MayaLiveLinkPlugin/MayaLiveLinkPlugin2016.Target.cs b/Engine/Source/Programs/MayaLiveLinkPlugin/MayaLiveLinkPlugin2016.Target.cs
--- a/Engine/Source/Programs/MayaLiveLinkPlugin/MayaLiveLinkPlugin2016.Target.cs
+++ b/Engine/Source/Programs/MayaLiveLinkPlugin/MayaLiveLinkPlugin2016.Target.cs
@@ -1,11 +1,11 @@
-﻿// Copyright 1998-2019 Epic Games, Inc. All Rights Reserved.
+// Copyright 1998-2019 Epic Games, Inc. All Rights Reserved.
 
 using UnrealBuildTool;
 using System.Collections.Generic;
 
 public abstract class MayaLiveLinkPluginTargetBase : TargetRules
 {
-	public MayaLiveLinkPluginTargetBase(TargetInfo Target) : base(Target)
+	public MayaLiveLinkPluginTargetBase(TargetInfo Target, string LaunchModuleName) : base(Target)
 	{
 		Type = TargetType.Program;
 
@@ -18,14 +18,16 @@ public abstract class MayaLiveLinkPluginTargetBase : TargetRules
 		bBuildWithEditorOnlyData = true;
 		bCompileAgainstEngine = false;
 		bCompileAgainstCoreUObject = true;
-        bCompileICU = false;
-        bHasExports = true;
+		bCompileICU = false;
+		bHasExports = true;
 
 		bBuildInSolutionByDefault = false;
 
+		this.LaunchModuleName = LaunchModuleName;
+
 		// Add a post-build step that copies the output to a file with the .mll extension
-		string OutputName = "MayaLiveLinkPlugin";
-		if(Target.Configuration != UnrealTargetConfiguration.Development)
+		string OutputName = LaunchModuleName;
+		if (Target.Configuration != UnrealTargetConfiguration.Development)
 		{
 			OutputName = string.Format("{0}-{1}-{2}", OutputName, Target.Platform, Target.Configuration);
 		}
@@ -35,8 +37,7 @@ public abstract class MayaLiveLinkPluginTargetBase : TargetRules
 
 public class MayaLiveLinkPlugin2016Target : MayaLiveLinkPluginTargetBase
 {
-	public MayaLiveLinkPlugin2016Target(TargetInfo Target) : base(Target)
+	public MayaLiveLinkPlugin2016Target(TargetInfo Target) : base(Target, "MayaLiveLinkPlugin2016")
 	{
-		LaunchModuleName = "MayaLiveLinkPlugin2016";
 	}
 }
\ No newline at end of file
diff --git a/Engine/Source/Programs/MayaLiveLinkPlugin/MayaLiveLinkPlugin2017.Target.cs b/Engine/Source/Programs/MayaLiveLinkPlugin/MayaLiveLinkPlugin2017.Target.cs
--- a/Engine/Source/Programs/MayaLiveLinkPlugin/MayaLiveLinkPlugin2017.Target.cs
+++ b/Engine/Source/Programs/MayaLiveLinkPlugin/MayaLiveLinkPlugin2017.Target.cs
@@ -3,8 +3,7 @@ using UnrealBuildTool;
 
 public class MayaLiveLinkPlugin2017Target : MayaLiveLinkPluginTargetBase
 {
-	public MayaLiveLinkPlugin2017Target(TargetInfo Target) : base(Target)
+	public MayaLiveLinkPlugin2017Target(TargetInfo Target) : base(Target, "MayaLiveLinkPlugin2017")
 	{
-		LaunchModuleName = "MayaLiveLinkPlugin2017";
 	}
 }
\ No newline at end of file
diff --git a/Engine/Source/Programs/MayaLiveLinkPlugin/MayaLiveLinkPlugin2018.Target.cs b/Engine/Source/Programs/MayaLiveLinkPlugin/MayaLiveLinkPlugin2018.Target.cs
--- a/Engine/Source/Programs/MayaLiveLinkPlugin/MayaLiveLinkPlugin2018.Target.cs
+++ b/Engine/Source/Programs/MayaLiveLinkPlugin/MayaLiveLinkPlugin2018.Target.cs
@@ -3,8 +3,7 @@ using UnrealBuildTool;
 
 public class MayaLiveLinkPlugin2018Target : MayaLiveLinkPluginTargetBase
 {
-	public MayaLiveLinkPlugin2018Target(TargetInfo Target) : base(Target)
+	public MayaLiveLinkPlugin2018Target(TargetInfo Target) : base(Target, "MayaLiveLinkPlugin2018")
 	{
-		LaunchModuleName = "MayaLiveLinkPlugin2018";
 	}
 }
\ No newline at end of file
diff --git a/Engine/Source/Programs/MayaLiveLinkPlugin/MayaLiveLinkPlugin2019.Target.cs b/Engine/Source/Programs/MayaLiveLinkPlugin/MayaLiveLinkPlugin2019.Target.cs
--- a/Engine/Source/Programs/MayaLiveLinkPlugin/MayaLiveLinkPlugin2019.Target.cs
+++ b/Engine/Source/Programs/MayaLiveLinkPlugin/MayaLiveLinkPlugin2019.Target.cs
@@ -3,8 +3,7 @@ using UnrealBuildTool;
 
 public class MayaLiveLinkPlugin2019Target : MayaLiveLinkPluginTargetBase
 {
-	public MayaLiveLinkPlugin2019Target(TargetInfo Target) : base(Target)
+	public MayaLiveLinkPlugin2019Target(TargetInfo Target) : base(Target, "MayaLiveLinkPlugin2019")
 	{
-		LaunchModuleName = "MayaLiveLinkPlugin2019";
 	}
 }
\ No newline at end of file
