Remove CIS warning from MayaLiveLinkPlugin.
Remove dependency from TargetInfo for MobuLiveLinkPlugin.

diff --git a/Engine/Source/Programs/MayaLiveLinkPlugin/MayaLiveLinkPlugin2016.Build.cs b/Engine/Source/Programs/MayaLiveLinkPlugin/MayaLiveLinkPlugin2016.Build.cs
--- a/Engine/Source/Programs/MayaLiveLinkPlugin/MayaLiveLinkPlugin2016.Build.cs
+++ b/Engine/Source/Programs/MayaLiveLinkPlugin/MayaLiveLinkPlugin2016.Build.cs
@@ -28,7 +28,7 @@ public abstract class MayaLiveLinkPluginBase : ModuleRules
 		//
 
 		{
-			string MayaVersionString = GetMayaVersion();
+			//string MayaVersionString = GetMayaVersion();
 			string MayaInstallFolder = GetMayaInstallFolderPath();
 
 			// Make sure this version of Maya is actually installed
@@ -55,10 +55,10 @@ public abstract class MayaLiveLinkPluginBase : ModuleRules
 					);
 				}
 			}
-			else
-			{
-				// throw new BuildException("Couldn't find Autodesk Maya " + MayaVersionString + " in folder '" + MayaInstallFolder + "'.  This version of Maya must be installed for us to find the Maya SDK files.");
-			}
+			//else
+			//{
+			//	throw new BuildException("Couldn't find Autodesk Maya " + MayaVersionString + " in folder '" + MayaInstallFolder + "'.  This version of Maya must be installed for us to find the Maya SDK files.");
+			//}
 		}
 	}
 
