Fixed incorrect transforms being streamed from Maya plugin and allow to select role of streamed Subjects

diff --git a/Engine/Source/Programs/MayaLiveLinkPlugin/MayaLiveLinkPlugin.cpp b/Engine/Source/Programs/MayaLiveLinkPlugin/MayaLiveLinkPlugin.cpp
--- a/Engine/Source/Programs/MayaLiveLinkPlugin/MayaLiveLinkPlugin.cpp
+++ b/Engine/Source/Programs/MayaLiveLinkPlugin/MayaLiveLinkPlugin.cpp
@@ -11,6 +11,8 @@
 #include "Roles/LiveLinkAnimationTypes.h"
 #include "Roles/LiveLinkCameraRole.h"
 #include "Roles/LiveLinkCameraTypes.h"
+#include "Roles/LiveLinkLightRole.h"
+#include "Roles/LiveLinkLightTypes.h"
 #include "Roles/LiveLinkTransformRole.h"
 #include "Roles/LiveLinkTransformTypes.h"
 #include "LiveLinkProvider.h"
@@ -41,10 +43,14 @@ IMPLEMENT_APPLICATION(MayaLiveLinkPlugin, "MayaLiveLinkPlugin");
 #include <maya/MMatrix.h>
 #include <maya/MTransformationMatrix.h>
 #include <maya/MQuaternion.h>
+#include <maya/MStringArray.h>
+#include <maya/MString.h>
 #include <maya/MVector.h>
 #include <maya/MFnTransform.h>
 #include <maya/MFnIkJoint.h>
 #include <maya/MFnCamera.h>
+#include <maya/MFnLight.h>
+#include <maya/MFnSpotLight.h>
 #include <maya/MEulerRotation.h>
 #include <maya/MSelectionList.h>
 #include <maya/MAnimControl.h>
@@ -101,6 +107,12 @@ double RadToDeg(double Rad)
 	return (Rad*180.0) / E_PI;
 }
 
+double DegToRad(double Deg)
+{
+	const double E_PI = 3.1415926535897932384626433832795028841971693993751058209749445923078164062;
+	return Deg * (E_PI / 180.0f);
+}
+
 MMatrix GetScale(const MFnIkJoint& Joint)
 {
 	double Scale[3];
@@ -145,6 +157,13 @@ MMatrix GetTranslation(const MFnIkJoint& Joint)
 	return M.asMatrix();
 }
 
+void RotateCoordinateSystemForUnreal(MMatrix& InOutMatrix)
+{
+	MQuaternion RotOffset;
+	RotOffset.setToXAxis(DegToRad(90.0));
+	InOutMatrix *= RotOffset.asMatrix();
+}
+
 FTransform BuildUETransformFromMayaTransform(MMatrix& InMatrix)
 {
 	MMatrix UnrealSpaceJointMatrix;
@@ -173,7 +192,6 @@ FTransform BuildUETransformFromMayaTransform(MMatrix& InMatrix)
 
 	MTransformationMatrix UnrealSpaceJointTransform(UnrealSpaceJointMatrix);
 
-
 	// getRotation is MSpace::kTransform
 	double tx, ty, tz, tw;
 	UnrealSpaceJointTransform.getRotationQuaternion(tx, ty, tz, tw, MSpace::kWorld);
@@ -190,6 +208,16 @@ FTransform BuildUETransformFromMayaTransform(MMatrix& InMatrix)
 	return UETrans;
 }
 
+FColor MayaColorToUnreal(MColor Color)
+{
+	FColor Result;
+	Result.R = FMath::Clamp(Color[0] * 255.0, 0.0, 255.0);
+	Result.G = FMath::Clamp(Color[1] * 255.0, 0.0, 255.0);
+	Result.B = FMath::Clamp(Color[2] * 255.0, 0.0, 255.0);
+	Result.A = 255;
+	return Result;
+}
+
 void OutputRotation(const MMatrix& M)
 {
 	MTransformationMatrix TM(M);
@@ -210,10 +238,14 @@ public:
 	virtual ~IStreamedEntity() {};
 
 	virtual bool ShouldDisplayInUI() const { return false; }
-	virtual MString GetDisplayText() const = 0;
+	virtual MDagPath GetDagPath() const = 0;
+	virtual MString GetNameDisplayText() const = 0;
+	virtual MString GetRoleDisplayText() const = 0;
+	virtual MString GetSubjectTypeDisplayText() const = 0;
 	virtual bool ValidateSubject() const = 0;
 	virtual void RebuildSubjectData() = 0;
 	virtual void OnStream(double StreamTime, int32 FrameNumber) = 0;
+	virtual void SetStreamType(const FString& StreamType) = 0;
 };
 
 struct FStreamHierarchy
@@ -237,14 +269,15 @@ struct FStreamHierarchy
 	{}
 };
 
-struct FLiveLinkStreamedJointHeirarchySubject : IStreamedEntity
+struct FLiveLinkStreamedJointHierarchySubject : IStreamedEntity
 {
-	FLiveLinkStreamedJointHeirarchySubject(FName InSubjectName, MDagPath InRootPath)
+	FLiveLinkStreamedJointHierarchySubject(FName InSubjectName, MDagPath InRootPath)
 		: SubjectName(InSubjectName)
 		, RootDagPath(InRootPath)
+		, StreamMode(FCharacterStreamMode::FullHierarchy)
 	{}
 
-	~FLiveLinkStreamedJointHeirarchySubject()
+	~FLiveLinkStreamedJointHierarchySubject()
 	{
 		if (LiveLinkProvider.IsValid())
 		{
@@ -252,10 +285,13 @@ struct FLiveLinkStreamedJointHeirarchySubject : IStreamedEntity
 		}
 	}
 
-	virtual bool ShouldDisplayInUI() const { return true; }
-	virtual MString GetDisplayText() const { return MString("Character: ") + MString(*SubjectName.ToString()) + " ( " + RootDagPath.fullPathName() + " )"; }
+	virtual bool ShouldDisplayInUI() const override { return true; }
+	virtual MDagPath GetDagPath() const override { return RootDagPath; }
+	virtual MString GetNameDisplayText() const override { return MString(*SubjectName.ToString()); }
+	virtual MString GetRoleDisplayText() const override { return MString(*(CharacterStreamOptions[StreamMode])); }
+	virtual MString GetSubjectTypeDisplayText() const override { return MString("Character"); }
 
-	virtual bool ValidateSubject() const
+	virtual bool ValidateSubject() const override
 	{
 		MStatus Status;
 		bool bIsValid = RootDagPath.isValid(&Status);
@@ -300,99 +336,150 @@ struct FLiveLinkStreamedJointHeirarchySubject : IStreamedEntity
 		return bIsValid;
 	}
 
-	virtual void RebuildSubjectData()
+	virtual void RebuildSubjectData() override
 	{
-		JointsToStream.Reset();
+		if (StreamMode == FCharacterStreamMode::RootOnly)
+		{
+			FLiveLinkStaticDataStruct StaticData(FLiveLinkTransformStaticData::StaticStruct());
+			LiveLinkProvider->UpdateSubjectStaticData(SubjectName, ULiveLinkTransformRole::StaticClass(), MoveTemp(StaticData));
+		}
+		else if (StreamMode == FCharacterStreamMode::FullHierarchy)
+		{
+			JointsToStream.Reset();
 
-		FLiveLinkStaticDataStruct StaticData(FLiveLinkSkeletonStaticData::StaticStruct());
-		FLiveLinkSkeletonStaticData& AnimationData = *StaticData.Cast<FLiveLinkSkeletonStaticData>();
+			FLiveLinkStaticDataStruct StaticData(FLiveLinkSkeletonStaticData::StaticStruct());
+			FLiveLinkSkeletonStaticData& AnimationData = *StaticData.Cast<FLiveLinkSkeletonStaticData>();
 
-		MItDag::TraversalType traversalType = MItDag::kBreadthFirst;
-		MFn::Type filter = MFn::kJoint;
+			MItDag::TraversalType traversalType = MItDag::kBreadthFirst;
+			MFn::Type filter = MFn::kJoint;
 
-		MStatus status;
-		MItDag JointIterator;
-		JointIterator.reset(RootDagPath, MItDag::kDepthFirst, MFn::kJoint);
+			MStatus status;
+			MItDag JointIterator;
+			JointIterator.reset(RootDagPath, MItDag::kDepthFirst, MFn::kJoint);
 
-		//Build Hierarchy
-		TArray<int32> ParentIndexStack;
-		ParentIndexStack.SetNum(100, false);
+			//Build Hierarchy
+			TArray<int32> ParentIndexStack;
+			ParentIndexStack.SetNum(100, false);
 
-		int32 Index = 0;
+			int32 Index = 0;
 
-		for (; !JointIterator.isDone(); JointIterator.next())
-		{
-			uint32 Depth = JointIterator.depth();
-			if (Depth >= (uint32)ParentIndexStack.Num())
+			for (; !JointIterator.isDone(); JointIterator.next())
 			{
-				ParentIndexStack.SetNum(Depth + 1);
-			}
-			ParentIndexStack[Depth] = Index++;
+				uint32 Depth = JointIterator.depth();
+				if (Depth >= (uint32)ParentIndexStack.Num())
+				{
+					ParentIndexStack.SetNum(Depth + 1);
+				}
+				ParentIndexStack[Depth] = Index++;
 
-			int32 ParentIndex = Depth == 0 ? -1 : ParentIndexStack[Depth - 1];
+				int32 ParentIndex = Depth == 0 ? -1 : ParentIndexStack[Depth - 1];
 
-			MDagPath JointPath;
-			status = JointIterator.getPath(JointPath);
-			MFnIkJoint JointObject(JointPath);
+				MDagPath JointPath;
+				status = JointIterator.getPath(JointPath);
+				MFnIkJoint JointObject(JointPath);
 
-			//MGlobal::displayInfo(MString("Iter: ") + JointPath.fullPathName() + JointIterator.depth());
+				FName JointName(JointObject.name().asChar());
 
-			FName JointName(JointObject.name().asChar());
+				JointsToStream.Add(FStreamHierarchy(JointName, JointPath, ParentIndex));
+				AnimationData.BoneNames.Add(JointName);
+				AnimationData.BoneParents.Add(ParentIndex);
+			}
 
-			JointsToStream.Add(FStreamHierarchy(JointName, JointPath, ParentIndex));
-			AnimationData.BoneNames.Add(JointName);
-			AnimationData.BoneParents.Add(ParentIndex);
+			LiveLinkProvider->UpdateSubjectStaticData(SubjectName, ULiveLinkAnimationRole::StaticClass(), MoveTemp(StaticData));
 		}
-
-		LiveLinkProvider->UpdateSubjectStaticData(SubjectName, ULiveLinkAnimationRole::StaticClass(), MoveTemp(StaticData));
 	}
 
-	virtual void OnStream(double StreamTime, int32 FrameNumber)
+	virtual void OnStream(double StreamTime, int32 FrameNumber) override
 	{
-		FLiveLinkFrameDataStruct FrameData(FLiveLinkAnimationFrameData::StaticStruct());
-		FLiveLinkAnimationFrameData& AnimationData = *FrameData.Cast<FLiveLinkAnimationFrameData>();
+		if (StreamMode == FCharacterStreamMode::RootOnly)
+		{
+			MFnTransform TransformNode(RootDagPath);
 
-		AnimationData.Transforms.Reserve(JointsToStream.Num());
+			MMatrix Transform = TransformNode.transformation().asMatrix();
 
-		TArray<MMatrix> InverseScales;
-		InverseScales.Reserve(JointsToStream.Num());
+			FLiveLinkFrameDataStruct FrameData(FLiveLinkTransformFrameData::StaticStruct());
+			FLiveLinkTransformFrameData& TransformData = *FrameData.Cast<FLiveLinkTransformFrameData>();
 
-		for (int32 Idx = 0; Idx < JointsToStream.Num(); ++Idx)
+			RotateCoordinateSystemForUnreal(Transform);
+
+			// Convert Maya Camera orientation to Unreal
+			TransformData.Transform = BuildUETransformFromMayaTransform(Transform);
+
+			TransformData.WorldTime = StreamTime;
+			LiveLinkProvider->UpdateSubjectFrameData(SubjectName, MoveTemp(FrameData));
+		}
+		else if (StreamMode == FCharacterStreamMode::FullHierarchy)
 		{
-			const FStreamHierarchy& H = JointsToStream[Idx];
+			FLiveLinkFrameDataStruct FrameData(FLiveLinkAnimationFrameData::StaticStruct());
+			FLiveLinkAnimationFrameData& AnimationData = *FrameData.Cast<FLiveLinkAnimationFrameData>();
+
+			AnimationData.Transforms.Reserve(JointsToStream.Num());
+
+			TArray<MMatrix> InverseScales;
+			InverseScales.Reserve(JointsToStream.Num());
+
+			for (int32 Idx = 0; Idx < JointsToStream.Num(); ++Idx)
+			{
+				const FStreamHierarchy& H = JointsToStream[Idx];
+
+				MTransformationMatrix::RotationOrder RotOrder = H.JointObject.rotationOrder();
 
-			MTransformationMatrix::RotationOrder RotOrder = H.JointObject.rotationOrder();
+				MMatrix JointScale = GetScale(H.JointObject);
+				InverseScales.Add(JointScale.inverse());
 
-			MMatrix JointScale = GetScale(H.JointObject);
-			InverseScales.Add(JointScale.inverse());
+				MMatrix ParentInverseScale = (H.ParentIndex == -1) ? MMatrix::identity : InverseScales[H.ParentIndex];
 
-			MMatrix ParentInverseScale = (H.ParentIndex == -1) ? MMatrix::identity : InverseScales[H.ParentIndex];
+				MMatrix MayaSpaceJointMatrix = JointScale *
+					GetRotationOrientation(H.JointObject, RotOrder) *
+					GetRotation(H.JointObject, RotOrder) *
+					GetJointOrientation(H.JointObject, RotOrder) *
+					ParentInverseScale *
+					GetTranslation(H.JointObject);
 
-			MMatrix MayaSpaceJointMatrix = JointScale *
-				GetRotationOrientation(H.JointObject, RotOrder) *
-				GetRotation(H.JointObject, RotOrder) *
-				GetJointOrientation(H.JointObject, RotOrder) *
-				ParentInverseScale *
-				GetTranslation(H.JointObject);
+				if (Idx == 0) // rotate the root joint to get the correct character rotation in Unreal
+				{
+					RotateCoordinateSystemForUnreal(MayaSpaceJointMatrix);
+				}
 
-			AnimationData.Transforms.Add(BuildUETransformFromMayaTransform(MayaSpaceJointMatrix));
+				AnimationData.Transforms.Add(BuildUETransformFromMayaTransform(MayaSpaceJointMatrix));
+			}
+
+			AnimationData.WorldTime = StreamTime;
+			LiveLinkProvider->UpdateSubjectFrameData(SubjectName, MoveTemp(FrameData));
 		}
+	}
 
-		AnimationData.WorldTime = StreamTime;
-		LiveLinkProvider->UpdateSubjectFrameData(SubjectName, MoveTemp(FrameData));
+	virtual void SetStreamType(const FString& StreamTypeIn) override
+	{
+		for (int32 StreamTypeIdx = 0; StreamTypeIdx < CharacterStreamOptions.Num(); ++StreamTypeIdx)
+		{
+			if (CharacterStreamOptions[StreamTypeIdx] == StreamTypeIn && StreamMode != (FCharacterStreamMode)StreamTypeIdx)
+			{
+				StreamMode = (FCharacterStreamMode)StreamTypeIdx;
+				RebuildSubjectData();
+				return;
+			}
+		}
 	}
 
 private:
 	FName SubjectName;
 	MDagPath RootDagPath;
-
 	TArray<FStreamHierarchy> JointsToStream;
+
+	const TArray<FString> CharacterStreamOptions = { TEXT("Root Only"), TEXT("Full Hierarchy") };
+	enum FCharacterStreamMode
+	{
+		RootOnly,
+		FullHierarchy,
+	};
+	FCharacterStreamMode StreamMode;
 };
 
 struct FLiveLinkBaseCameraStreamedSubject : public IStreamedEntity
 {
 public:
-	FLiveLinkBaseCameraStreamedSubject(FName InSubjectName) : SubjectName(InSubjectName) {}
+	FLiveLinkBaseCameraStreamedSubject(FName InSubjectName) : SubjectName(InSubjectName), StreamMode(FCameraStreamMode::Camera) {}
 
 	~FLiveLinkBaseCameraStreamedSubject()
 	{
@@ -402,17 +489,39 @@ public:
 		}
 	}
 
-	virtual bool ValidateSubject() const { return true; }
+	virtual MString GetNameDisplayText() const override { return *SubjectName.ToString(); }
+	virtual MString GetRoleDisplayText() const override { return MString(*(CameraStreamOptions[StreamMode])); }
+	virtual MString GetSubjectTypeDisplayText() const override { return MString("Camera"); }
+
+	virtual bool ValidateSubject() const override { return true; }
 
-	virtual void RebuildSubjectData()
+	virtual void RebuildSubjectData() override
 	{
-		FLiveLinkStaticDataStruct StaticData(FLiveLinkCameraStaticData::StaticStruct());
-		FLiveLinkCameraStaticData& CameraData = *StaticData.Cast<FLiveLinkCameraStaticData>();
-		CameraData.bIsFieldOfViewSupported = true;
-		CameraData.bIsAspectRatioSupported = true;
-		CameraData.bIsFocalLengthSupported = true;
-		CameraData.bIsProjectionModeSupported = true;
-		LiveLinkProvider->UpdateSubjectStaticData(SubjectName, ULiveLinkCameraRole::StaticClass(), MoveTemp(StaticData));
+		if (StreamMode == FCameraStreamMode::RootOnly)
+		{
+			FLiveLinkStaticDataStruct TransformData(FLiveLinkTransformStaticData::StaticStruct());
+			LiveLinkProvider->UpdateSubjectStaticData(SubjectName, ULiveLinkTransformRole::StaticClass(), MoveTemp(TransformData));
+		}
+		else if (StreamMode == FCameraStreamMode::FullHierarchy)
+		{
+			FLiveLinkStaticDataStruct StaticData(FLiveLinkSkeletonStaticData::StaticStruct());
+			FLiveLinkSkeletonStaticData& AnimationData = *StaticData.Cast<FLiveLinkSkeletonStaticData>();
+
+			AnimationData.BoneNames.Add(FName("root"));
+			AnimationData.BoneParents.Add(-1);
+
+			LiveLinkProvider->UpdateSubjectStaticData(SubjectName, ULiveLinkAnimationRole::StaticClass(), MoveTemp(StaticData));
+		}
+		else if (StreamMode == FCameraStreamMode::Camera)
+		{
+			FLiveLinkStaticDataStruct StaticData(FLiveLinkCameraStaticData::StaticStruct());
+			FLiveLinkCameraStaticData& CameraData = *StaticData.Cast<FLiveLinkCameraStaticData>();
+			CameraData.bIsFieldOfViewSupported = true;
+			CameraData.bIsAspectRatioSupported = true;
+			CameraData.bIsFocalLengthSupported = true;
+			CameraData.bIsProjectionModeSupported = true;
+			LiveLinkProvider->UpdateSubjectStaticData(SubjectName, ULiveLinkCameraRole::StaticClass(), MoveTemp(StaticData));
+		}
 	}
 
 	void StreamCamera(MDagPath CameraPath, double StreamTime, int32 FrameNumber)
@@ -432,31 +541,71 @@ public:
 			SetMatrixRow(CameraTransformMatrix[2], C.upDirection(MSpace::kWorld));
 			SetMatrixRow(CameraTransformMatrix[3], EyeLocation);
 
-			FLiveLinkFrameDataStruct FrameData(FLiveLinkCameraFrameData::StaticStruct());
-			FLiveLinkCameraFrameData& CameraData = *FrameData.Cast<FLiveLinkCameraFrameData>();
+			RotateCoordinateSystemForUnreal(CameraTransformMatrix);
+			FTransform CameraTransform = BuildUETransformFromMayaTransform(CameraTransformMatrix);
+			// Convert Maya Camera orientation to Unreal
+			CameraTransform.SetRotation(CameraTransform.GetRotation() * FRotator(0.0f, -90.0f, 0.0f).Quaternion());
 
-			CameraData.FieldOfView = C.horizontalFieldOfView();
-			CameraData.AspectRatio = C.aspectRatio();
-			CameraData.FocalLength = C.focalLength();
-			CameraData.ProjectionMode = C.isOrtho() ? ELiveLinkCameraProjectionMode::Orthographic : ELiveLinkCameraProjectionMode::Perspective;
+			if (StreamMode == FCameraStreamMode::RootOnly)
+			{
+				FLiveLinkFrameDataStruct TransformData = (FLiveLinkTransformFrameData::StaticStruct());
+				FLiveLinkTransformFrameData& CameraTransformData = *TransformData.Cast<FLiveLinkTransformFrameData>();
+				CameraTransformData.Transform = CameraTransform;
+				CameraTransformData.WorldTime = StreamTime;
+				LiveLinkProvider->UpdateSubjectFrameData(SubjectName, MoveTemp(TransformData));
+			}
+			else if (StreamMode == FCameraStreamMode::FullHierarchy)
+			{
+				FLiveLinkFrameDataStruct FrameData(FLiveLinkAnimationFrameData::StaticStruct());
+				FLiveLinkAnimationFrameData& AnimationData = *FrameData.Cast<FLiveLinkAnimationFrameData>();
 
-			CameraData.Transform = BuildUETransformFromMayaTransform(CameraTransformMatrix);
-			// Convert Maya Camera orientation to Unreal
-			CameraData.Transform.SetRotation(CameraData.Transform.GetRotation() * FRotator(0.f, -90.f, 0.f).Quaternion());
-			CameraData.WorldTime = StreamTime;
+				AnimationData.Transforms.Add(CameraTransform);
+				AnimationData.WorldTime = StreamTime;
+				LiveLinkProvider->UpdateSubjectFrameData(SubjectName, MoveTemp(FrameData));
+			}
+			else if (StreamMode == FCameraStreamMode::Camera)
+			{
+				FLiveLinkFrameDataStruct FrameData(FLiveLinkCameraFrameData::StaticStruct());
+				FLiveLinkCameraFrameData& CameraData = *FrameData.Cast<FLiveLinkCameraFrameData>();
 
-			LiveLinkProvider->UpdateSubjectFrameData(SubjectName, MoveTemp(FrameData));
+				CameraData.FieldOfView = C.horizontalFieldOfView();
+				CameraData.AspectRatio = C.aspectRatio();
+				CameraData.FocalLength = C.focalLength();
+				CameraData.ProjectionMode = C.isOrtho() ? ELiveLinkCameraProjectionMode::Orthographic : ELiveLinkCameraProjectionMode::Perspective;
+
+				CameraData.Transform = CameraTransform;
+				CameraData.WorldTime = StreamTime;
+
+				LiveLinkProvider->UpdateSubjectFrameData(SubjectName, MoveTemp(FrameData));
+			}
+		}
+	}
+
+	virtual void SetStreamType(const FString& StreamTypeIn) override
+	{
+		for (int32 StreamTypeIdx = 0; StreamTypeIdx < CameraStreamOptions.Num(); ++StreamTypeIdx)
+		{
+			if (CameraStreamOptions[StreamTypeIdx] == StreamTypeIn && StreamMode != (FCameraStreamMode)StreamTypeIdx)
+			{
+				StreamMode = (FCameraStreamMode)StreamTypeIdx;
+				RebuildSubjectData();
+				return;
+			}
 		}
 	}
 
 protected:
 	FName  SubjectName;
-	static TArray<FName> ActiveCameraBoneNames;
-	static TArray<int32> ActiveCameraBoneParents;
-};
 
-TArray<FName> FLiveLinkBaseCameraStreamedSubject::ActiveCameraBoneNames = { FName("root") };
-TArray<int32> FLiveLinkBaseCameraStreamedSubject::ActiveCameraBoneParents = { -1 };
+	const TArray<FString> CameraStreamOptions = { TEXT("Root Only"), TEXT("Full Hierarchy"), TEXT("Camera") };
+	enum FCameraStreamMode
+	{
+		RootOnly,
+		FullHierarchy,
+		Camera,
+	};
+	FCameraStreamMode StreamMode;
+};
 
 struct FLiveLinkStreamedActiveCamera : public FLiveLinkBaseCameraStreamedSubject
 {
@@ -465,9 +614,9 @@ public:
 
 	MDagPath CurrentActiveCameraDag;
 
-	virtual MString GetDisplayText() const { return MString(); }
+	virtual MDagPath GetDagPath() const override { return CurrentActiveCameraDag; }
 
-	virtual void OnStream(double StreamTime, int32 FrameNumber)
+	virtual void OnStream(double StreamTime, int32 FrameNumber) override
 	{
 		MStatus Status;
 		M3dView ActiveView = M3dView::active3dView(&Status);
@@ -487,15 +636,18 @@ private:
 	static FName ActiveCameraName;
 };
 
+FName FLiveLinkStreamedActiveCamera::ActiveCameraName("EditorActiveCamera");
+
 struct FLiveLinkStreamedCameraSubject : FLiveLinkBaseCameraStreamedSubject
 {
 public:
+	virtual MDagPath GetDagPath() const override { return CameraPath; }
+
 	FLiveLinkStreamedCameraSubject(FName InSubjectName, MDagPath InDagPath) : FLiveLinkBaseCameraStreamedSubject(InSubjectName), CameraPath(InDagPath) {}
 
-	virtual bool ShouldDisplayInUI() const { return true; }
-	virtual MString GetDisplayText() const { return MString("Camera: ") + *SubjectName.ToString() + " ( " + CameraPath.fullPathName() + " )"; }
+	virtual bool ShouldDisplayInUI() const override { return true; }
 
-	virtual void OnStream(double StreamTime, int32 FrameNumber)
+	virtual void OnStream(double StreamTime, int32 FrameNumber) override
 	{
 		StreamCamera(CameraPath, StreamTime, FrameNumber);
 	}
@@ -504,7 +656,136 @@ private:
 	MDagPath CameraPath;
 };
 
-FName FLiveLinkStreamedActiveCamera::ActiveCameraName("EditorActiveCamera");
+struct FLiveLinkStreamedLightSubject : IStreamedEntity
+{
+public:
+	FLiveLinkStreamedLightSubject(FName InSubjectName, MDagPath InRootPath)
+		: SubjectName(InSubjectName)
+		, RootDagPath(InRootPath)
+		, StreamMode(FLightStreamMode::Light)
+	{}
+
+	~FLiveLinkStreamedLightSubject()
+	{
+		if (LiveLinkProvider.IsValid())
+		{
+			LiveLinkProvider->RemoveSubject(SubjectName);
+		}
+	}
+
+	virtual bool ShouldDisplayInUI() const override { return true; }
+	virtual MDagPath GetDagPath() const override { return RootDagPath; }
+	virtual MString GetNameDisplayText() const override { return MString(*SubjectName.ToString()); }
+	virtual MString GetRoleDisplayText() const override { return MString(*(LightStreamOptions[StreamMode])); }
+	virtual MString GetSubjectTypeDisplayText() const override { return MString("Light"); }
+
+	virtual bool ValidateSubject() const override { return true; }
+
+	virtual void RebuildSubjectData() override
+	{
+		if (StreamMode == FLightStreamMode::RootOnly)
+		{
+			FLiveLinkStaticDataStruct TransformData(FLiveLinkTransformStaticData::StaticStruct());
+			LiveLinkProvider->UpdateSubjectStaticData(SubjectName, ULiveLinkTransformRole::StaticClass(), MoveTemp(TransformData));
+		}
+		else if (StreamMode == FLightStreamMode::FullHierarchy)
+		{
+			FLiveLinkStaticDataStruct StaticData(FLiveLinkSkeletonStaticData::StaticStruct());
+			FLiveLinkSkeletonStaticData& AnimationData = *StaticData.Cast<FLiveLinkSkeletonStaticData>();
+
+			AnimationData.BoneNames.Add(FName("root"));
+			AnimationData.BoneParents.Add(-1);
+
+			LiveLinkProvider->UpdateSubjectStaticData(SubjectName, ULiveLinkAnimationRole::StaticClass(), MoveTemp(StaticData));
+		}
+		else if (StreamMode == FLightStreamMode::Light)
+		{
+			FLiveLinkStaticDataStruct StaticData(FLiveLinkLightStaticData::StaticStruct());
+			FLiveLinkLightStaticData& LightData = *StaticData.Cast<FLiveLinkLightStaticData>();
+			LightData.bIsIntensitySupported = true;
+			LightData.bIsLightColorSupported = true;
+			const bool bIsSpotLight = RootDagPath.hasFn(MFn::kSpotLight);
+			LightData.bIsInnerConeAngleSupported = bIsSpotLight;
+			LightData.bIsOuterConeAngleSupported = bIsSpotLight;
+
+			LiveLinkProvider->UpdateSubjectStaticData(SubjectName, ULiveLinkLightRole::StaticClass(), MoveTemp(StaticData));
+		}
+	}
+
+	virtual void OnStream(double StreamTime, int32 FrameNumber) override
+	{
+		MFnTransform TransformNode(RootDagPath);
+		MMatrix MayaTransform = TransformNode.transformation().asMatrix();
+		RotateCoordinateSystemForUnreal(MayaTransform);
+		const FTransform UnrealTransform = BuildUETransformFromMayaTransform(MayaTransform);
+
+		if (StreamMode == FLightStreamMode::RootOnly)
+		{
+			FLiveLinkFrameDataStruct FrameData(FLiveLinkTransformFrameData::StaticStruct());
+			FLiveLinkTransformFrameData& TransformData = *FrameData.Cast<FLiveLinkTransformFrameData>();
+
+			TransformData.Transform = UnrealTransform;
+			TransformData.WorldTime = StreamTime;
+			LiveLinkProvider->UpdateSubjectFrameData(SubjectName, MoveTemp(FrameData));
+		}
+		else if (StreamMode == FLightStreamMode::FullHierarchy)
+		{
+			FLiveLinkFrameDataStruct FrameData(FLiveLinkAnimationFrameData::StaticStruct());
+			FLiveLinkAnimationFrameData& AnimationData = *FrameData.Cast<FLiveLinkAnimationFrameData>();
+
+			AnimationData.Transforms.Add(UnrealTransform);
+			AnimationData.WorldTime = StreamTime;
+			LiveLinkProvider->UpdateSubjectFrameData(SubjectName, MoveTemp(FrameData));
+		}
+		else if (StreamMode == FLightStreamMode::Light)
+		{
+
+			FLiveLinkFrameDataStruct FrameData(FLiveLinkLightFrameData::StaticStruct());
+			FLiveLinkLightFrameData& TransformData = *FrameData.Cast<FLiveLinkLightFrameData>();
+
+			TransformData.Transform = UnrealTransform;
+			TransformData.WorldTime = StreamTime;
+
+			MFnLight Light(RootDagPath);
+			TransformData.Intensity = Light.intensity();
+			TransformData.LightColor = MayaColorToUnreal(Light.color());
+
+			if (RootDagPath.hasFn(MFn::kSpotLight))
+			{
+				MFnSpotLight SpotLight(RootDagPath);
+				TransformData.InnerConeAngle = static_cast<float>(SpotLight.coneAngle());
+				TransformData.OuterConeAngle = static_cast<float>(SpotLight.penumbraAngle());
+			}
+			LiveLinkProvider->UpdateSubjectFrameData(SubjectName, MoveTemp(FrameData));
+		}
+	}
+
+	virtual void SetStreamType(const FString& StreamTypeIn) override
+	{
+		for (int32 StreamTypeIdx = 0; StreamTypeIdx < LightStreamOptions.Num(); ++StreamTypeIdx)
+		{
+			if (LightStreamOptions[StreamTypeIdx] == StreamTypeIn && StreamMode != (FLightStreamMode)StreamTypeIdx)
+			{
+				StreamMode = (FLightStreamMode)StreamTypeIdx;
+				RebuildSubjectData();
+				return;
+			}
+		}
+	}
+
+private:
+	FName SubjectName;
+	MDagPath RootDagPath;
+
+	const TArray<FString> LightStreamOptions = { TEXT("Root Only"), TEXT("Full Hierarchy"), TEXT("Light") };
+	enum FLightStreamMode
+	{
+		RootOnly,
+		FullHierarchy,
+		Light,
+	};
+	FLightStreamMode StreamMode;
+};
 
 struct FLiveLinkStreamedPropSubject : IStreamedEntity
 {
@@ -512,6 +793,7 @@ public:
 	FLiveLinkStreamedPropSubject(FName InSubjectName, MDagPath InRootPath)
 		: SubjectName(InSubjectName)
 		, RootDagPath(InRootPath)
+		, StreamMode(FPropStreamMode::RootOnly)
 	{}
 
 	~FLiveLinkStreamedPropSubject()
@@ -522,45 +804,86 @@ public:
 		}
 	}
 
-	virtual bool ShouldDisplayInUI() const { return true; }
-	virtual MString GetDisplayText() const { return MString("Prop: ") + MString(*SubjectName.ToString()) + " ( " + RootDagPath.fullPathName() + " )"; }
+	virtual bool ShouldDisplayInUI() const override { return true; }
+	virtual MDagPath GetDagPath() const override { return RootDagPath; }
+	virtual MString GetNameDisplayText() const override { return MString(*SubjectName.ToString()); }
+	virtual MString GetRoleDisplayText() const override { return MString(*(PropStreamOptions[StreamMode])); }
+	virtual MString GetSubjectTypeDisplayText() const override { return MString("Prop"); }
 
-	virtual bool ValidateSubject() const {return true;}
+	virtual bool ValidateSubject() const override {return true;}
 
-	virtual void RebuildSubjectData()
+	virtual void RebuildSubjectData() override
 	{
-		FLiveLinkStaticDataStruct StaticData(FLiveLinkTransformStaticData::StaticStruct());
-		LiveLinkProvider->UpdateSubjectStaticData(SubjectName, ULiveLinkTransformRole::StaticClass(), MoveTemp(StaticData));
+		if (StreamMode == FPropStreamMode::RootOnly)
+		{
+			FLiveLinkStaticDataStruct TransformData(FLiveLinkTransformStaticData::StaticStruct());
+			LiveLinkProvider->UpdateSubjectStaticData(SubjectName, ULiveLinkTransformRole::StaticClass(), MoveTemp(TransformData));
+		}
+		else if (StreamMode == FPropStreamMode::FullHierarchy)
+		{
+			FLiveLinkStaticDataStruct StaticData(FLiveLinkSkeletonStaticData::StaticStruct());
+			FLiveLinkSkeletonStaticData& AnimationData = *StaticData.Cast<FLiveLinkSkeletonStaticData>();
+
+			AnimationData.BoneNames.Add(FName("root"));
+			AnimationData.BoneParents.Add(-1);
+
+			LiveLinkProvider->UpdateSubjectStaticData(SubjectName, ULiveLinkAnimationRole::StaticClass(), MoveTemp(StaticData));
+		}
 	}
 
-	virtual void OnStream(double StreamTime, int32 FrameNumber)
+	virtual void OnStream(double StreamTime, int32 FrameNumber) override
 	{
 		MFnTransform TransformNode(RootDagPath);
+		MMatrix MayaTransform = TransformNode.transformation().asMatrix();
+		RotateCoordinateSystemForUnreal(MayaTransform);
+		const FTransform UnrealTransform = BuildUETransformFromMayaTransform(MayaTransform);
+
+		if (StreamMode == FPropStreamMode::RootOnly)
+		{
+			FLiveLinkFrameDataStruct FrameData(FLiveLinkTransformFrameData::StaticStruct());
+			FLiveLinkTransformFrameData& TransformData = *FrameData.Cast<FLiveLinkTransformFrameData>();
 
-		MMatrix Transform = TransformNode.transformation().asMatrix();
+			TransformData.Transform = UnrealTransform;
+			TransformData.WorldTime = StreamTime;
+			LiveLinkProvider->UpdateSubjectFrameData(SubjectName, MoveTemp(FrameData));
+		}
+		else if (StreamMode == FPropStreamMode::FullHierarchy)
+		{
+			FLiveLinkFrameDataStruct FrameData(FLiveLinkAnimationFrameData::StaticStruct());
+			FLiveLinkAnimationFrameData& AnimationData = *FrameData.Cast<FLiveLinkAnimationFrameData>();
 
-		FLiveLinkFrameDataStruct FrameData(FLiveLinkTransformFrameData::StaticStruct());
-		FLiveLinkTransformFrameData& TransformData = *FrameData.Cast<FLiveLinkTransformFrameData>();
+			AnimationData.Transforms.Add(UnrealTransform);
+			AnimationData.WorldTime = StreamTime;
+			LiveLinkProvider->UpdateSubjectFrameData(SubjectName, MoveTemp(FrameData));
+		}
+	}
 
-		// Convert Maya Camera orientation to Unreal
-		TransformData.Transform = BuildUETransformFromMayaTransform(Transform);
-		TransformData.WorldTime = StreamTime;
-		LiveLinkProvider->UpdateSubjectFrameData(SubjectName, MoveTemp(FrameData));
+	virtual void SetStreamType(const FString& StreamTypeIn) override
+	{
+		for (int32 StreamTypeIdx = 0; StreamTypeIdx < PropStreamOptions.Num(); ++StreamTypeIdx)
+		{
+			if (PropStreamOptions[StreamTypeIdx] == StreamTypeIn && StreamMode != (FPropStreamMode)StreamTypeIdx)
+			{
+				StreamMode = (FPropStreamMode)StreamTypeIdx;
+				RebuildSubjectData();
+				return;
+			}
+		}
 	}
 
 private:
 	FName SubjectName;
 	MDagPath RootDagPath;
-	bool bHasBeenRebuild;
-
-	static TArray<FName> PropBoneNames;
-	static TArray<int32> PropBoneParents;
 
+	const TArray<FString> PropStreamOptions = { TEXT("Root Only"), TEXT("Full Hierarchy") };
+	enum FPropStreamMode
+	{
+		RootOnly,
+		FullHierarchy,
+	};
+	FPropStreamMode StreamMode;
 };
 
-TArray<FName> FLiveLinkStreamedPropSubject::PropBoneNames = { FName("root") };
-TArray<int32> FLiveLinkStreamedPropSubject::PropBoneParents = { -1 };
-
 class FLiveLinkStreamedSubjectManager
 {
 private:
@@ -582,13 +905,46 @@ public:
 		Reset();
 	}
 
-	void GetSubjectEntries(TArray<MString>& Entries) const
+	void GetSubjectNames(TArray<MString>& Entries) const
 	{
 		for (const TSharedPtr<IStreamedEntity>& Subject : Subjects)
 		{
 			if (Subject->ShouldDisplayInUI())
 			{
-				Entries.Add(Subject->GetDisplayText());
+				Entries.Add(Subject->GetNameDisplayText());
+			}
+		}
+	}
+
+	void GetSubjectPaths(TArray<MString>& Entries) const
+	{
+		for (const TSharedPtr<IStreamedEntity>& Subject : Subjects)
+		{
+			if (Subject->ShouldDisplayInUI())
+			{
+				Entries.Add(Subject->GetDagPath().fullPathName());
+			}
+		}
+	}
+
+	void GetSubjectRoles(TArray<MString>& Entries) const
+	{
+		for (const TSharedPtr<IStreamedEntity>& Subject : Subjects)
+		{
+			if (Subject->ShouldDisplayInUI())
+			{
+				Entries.Add(Subject->GetRoleDisplayText());
+			}
+		}
+	}
+
+	void GetSubjectTypes(TArray<MString>& Entries) const
+	{
+		for (const TSharedPtr<IStreamedEntity>& Subject : Subjects)
+		{
+			if (Subject->ShouldDisplayInUI())
+			{
+				Entries.Add(Subject->GetSubjectTypeDisplayText());
 			}
 		}
 	}
@@ -609,7 +965,7 @@ public:
 
 	void AddJointHeirarchySubject(FName SubjectName, MDagPath RootPath)
 	{
-		AddSubjectOfType<FLiveLinkStreamedJointHeirarchySubject>(SubjectName, RootPath);
+		AddSubjectOfType<FLiveLinkStreamedJointHierarchySubject>(SubjectName, RootPath);
 	}
 
 	void AddCameraSubject(FName SubjectName, MDagPath RootPath)
@@ -617,18 +973,23 @@ public:
 		AddSubjectOfType<FLiveLinkStreamedCameraSubject>(SubjectName, RootPath);
 	}
 
+	void AddLightSubject(FName SubjectName, MDagPath RootPath)
+	{
+		AddSubjectOfType<FLiveLinkStreamedLightSubject>(SubjectName, RootPath);
+	}
+
 	void AddPropSubject(FName SubjectName, MDagPath RootPath)
 	{
 		AddSubjectOfType<FLiveLinkStreamedPropSubject>(SubjectName, RootPath);
 	}
 
-	void RemoveSubject(MString SubjectToRemove)
+	void RemoveSubject(MString PathOfSubjectToRemove)
 	{
 		for (int32 Index = Subjects.Num() - 1; Index >= 0; --Index)
 		{
 			if (Subjects[Index]->ShouldDisplayInUI())
 			{
-				if (Subjects[Index]->GetDisplayText() == SubjectToRemove)
+				if (Subjects[Index]->GetDagPath().fullPathName() == PathOfSubjectToRemove)
 				{
 					Subjects.RemoveAt(Index);
 					break;
@@ -637,6 +998,58 @@ public:
 		}
 	}
 
+	void ChangeSubjectName(MString SubjectDagPath, MString NewName)
+	{
+		if (IStreamedEntity* Subject = GetSubjectByDagPath(SubjectDagPath))
+		{
+			const MDagPath PathBackup = Subject->GetDagPath(); // store so we can re-create the subject
+			RemoveSubject(SubjectDagPath);
+
+			if (PathBackup.hasFn(MFn::kJoint))
+			{
+				AddJointHeirarchySubject(NewName.asChar(), PathBackup);
+			}
+			else if (PathBackup.hasFn(MFn::kCamera))
+			{
+				AddCameraSubject(NewName.asChar(), PathBackup);
+			}
+			else if (PathBackup.hasFn(MFn::kLight))
+			{
+				AddLightSubject(NewName.asChar(), PathBackup);
+			}
+			else
+			{
+				AddPropSubject(NewName.asChar(), PathBackup);
+			}
+		}
+	}
+
+	void ChangeStreamType(MString SubjectPathIn, MString StreamTypeIn)
+	{
+		if (IStreamedEntity* Subject = GetSubjectByDagPath(SubjectPathIn))
+		{
+			const FString StreamType(StreamTypeIn.asChar());
+			Subject->SetStreamType(StreamType);
+		}
+	}
+
+	IStreamedEntity* GetSubjectByDagPath(MString Path)
+	{
+		TSharedPtr<IStreamedEntity>* Found = Subjects.FindByPredicate([&Path](TSharedPtr<IStreamedEntity>& Subject)
+		{
+			if (Subject.IsValid())
+			{
+				if (Subject->GetDagPath().fullPathName() == Path)
+				{
+					return Subject;
+				}
+			}
+			return TSharedPtr<IStreamedEntity>();
+		});
+
+		return (*Found).Get();
+	}
+
 	void Reset()
 	{
 		Subjects.Reset();
@@ -664,20 +1077,20 @@ public:
 	}
 };
 
-const MString LiveLinkSubjectsCommandName("LiveLinkSubjects");
+const MString LiveLinkSubjectNamesCommandName("LiveLinkSubjectNames");
 
-class LiveLinkSubjectsCommand : public MPxCommand
+class LiveLinkSubjectNamesCommand : public MPxCommand
 {
 public:
 	static void		cleanup() {}
-	static void*	creator() { return new LiveLinkSubjectsCommand(); }
+	static void*	creator() { return new LiveLinkSubjectNamesCommand(); }
 
-	MStatus			doIt(const MArgList& args)
+	MStatus			doIt(const MArgList& args) override
 	{
-		TArray<MString> SubjectEntries;
-		LiveLinkStreamManager->GetSubjectEntries(SubjectEntries);
+		TArray<MString> SubjectNames;
+		LiveLinkStreamManager->GetSubjectNames(SubjectNames);
 
-		for (const MString& Entry : SubjectEntries)
+		for (const MString& Entry : SubjectNames)
 		{
 			appendToResult(Entry);
 		}
@@ -686,60 +1099,146 @@ public:
 	}
 };
 
-const MString LiveLinkAddSubjectCommandName("LiveLinkAddSubject");
+const MString LiveLinkSubjectPathsCommandName("LiveLinkSubjectPaths");
 
-class LiveLinkAddSubjectCommand : public MPxCommand
+class LiveLinkSubjectPathsCommand : public MPxCommand
 {
 public:
 	static void		cleanup() {}
-	static void*	creator() { return new LiveLinkAddSubjectCommand(); }
+	static void*	creator() { return new LiveLinkSubjectPathsCommand(); }
 
-	MStatus			doIt(const MArgList& args)
+	MStatus			doIt(const MArgList& args) override
 	{
-		MSyntax Syntax;
-		Syntax.addArg(MSyntax::kString);
+		TArray<MString> SubjectPaths;
+		LiveLinkStreamManager->GetSubjectPaths(SubjectPaths);
 
-		MArgDatabase argData(Syntax, args);
+		for (const MString& Entry : SubjectPaths)
+		{
+			appendToResult(Entry);
+		}
 
-		MString Name;
-		argData.getCommandArgument(0, Name);
+		return MS::kSuccess;
+	}
+};
 
-		FName SubjectFName(Name.asChar());
+const MString LiveLinkSubjectRolesCommandName("LiveLinkSubjectRoles");
+
+class LiveLinkSubjectRolesCommand : public MPxCommand
+{
+public:
+	static void		cleanup() {}
+	static void*	creator() { return new LiveLinkSubjectRolesCommand(); }
 
-		MSelectionList selected;
-		MGlobal::getActiveSelectionList(selected);
+	MStatus			doIt(const MArgList& args) override
+	{
+		TArray<MString> SubjectRoles;
+		LiveLinkStreamManager->GetSubjectRoles(SubjectRoles);
 
-		// Find selected joint
-		for (unsigned int i = 0; i < selected.length(); ++i)
+		for (const MString& Entry : SubjectRoles)
 		{
-			MObject obj;
+			appendToResult(Entry);
+		}
 
-			selected.getDependNode(i, obj);
+		return MS::kSuccess;
+	}
+};
 
-			if (obj.hasFn(MFn::kJoint))
-			{
-				MFnIkJoint JointObject(obj);
-				MDagPath Path;
-				JointObject.getPath(Path);
-				LiveLinkStreamManager->AddJointHeirarchySubject(SubjectFName, Path);
-			}
-			else if (obj.hasFn(MFn::kCamera))
+const MString LiveLinkSubjectTypesCommandName("LiveLinkSubjectTypes");
+
+class LiveLinkSubjectTypesCommand : public MPxCommand
+{
+public:
+	static void		cleanup() {}
+	static void*	creator() { return new LiveLinkSubjectTypesCommand(); }
+
+	MStatus			doIt(const MArgList& args) override
+	{
+		TArray<MString> SubjectTypes;
+		LiveLinkStreamManager->GetSubjectTypes(SubjectTypes);
+
+		for (const MString& Entry : SubjectTypes)
+		{
+			appendToResult(Entry);
+		}
+
+		return MS::kSuccess;
+	}
+};
+
+const MString LiveLinkAddSelectionCommandName("LiveLinkAddSelection");
+
+class LiveLinkAddSelectionCommand : public MPxCommand
+{
+public:
+	static void		cleanup() {}
+	static void*	creator() { return new LiveLinkAddSelectionCommand(); }
+
+	MStatus			doIt(const MArgList& args) override
+	{
+		MSelectionList SelectedItems;
+		MGlobal::getActiveSelectionList(SelectedItems);
+
+		for (unsigned int i = 0; i < SelectedItems.length(); ++i)
+		{
+			MObject SelectedRoot;
+			SelectedItems.getDependNode(i, SelectedRoot);
+
+			bool ItemAdded = false;
+
+			MItDag DagIterator;
+			DagIterator.reset(SelectedRoot);
+			// first try to find a specific subject item under the selected root item.
+			// we iterate through the DAG to find items in groups/sets and to be able to find the Shape compoments which
+			// hold the interesting properties.
+			while (!DagIterator.isDone() && !ItemAdded)
 			{
-				MFnCamera CameraObject(obj);
-				MDagPath Path;
-				CameraObject.getPath(Path);
-				LiveLinkStreamManager->AddCameraSubject(SubjectFName, Path);
+				MDagPath CurrentItemPath;
+				if (!DagIterator.getPath(CurrentItemPath))
+					continue;
+
+				MFnDagNode CurrentNode(CurrentItemPath);
+				const char* SubjectName = CurrentNode.name().asChar();
+
+				if (CurrentItemPath.hasFn(MFn::kJoint))
+				{
+					LiveLinkStreamManager->AddJointHeirarchySubject(SubjectName, CurrentItemPath);
+					ItemAdded = true;
+				}
+				else if (CurrentItemPath.hasFn(MFn::kCamera))
+				{
+					LiveLinkStreamManager->AddCameraSubject(SubjectName, CurrentItemPath);
+					ItemAdded = true;
+				}
+				else if (CurrentItemPath.hasFn(MFn::kLight))
+				{
+					LiveLinkStreamManager->AddLightSubject(SubjectName, CurrentItemPath);
+					ItemAdded = true;
+				}
+
+				if (ItemAdded)
+				{
+					MGlobal::displayInfo(MString("LiveLinkAddSubjectCommand ") + SubjectName);
+				}
+
+				DagIterator.next();
 			}
-			else if(obj.hasFn(MFn::kTransform))
+
+			// if there was no specific item, we assume that the selected item is a prop.
+			// the props are handled differently because almost everything has a kTransform function set, so if a subject
+			// is under a group node or in a set, the group would be added as a prop otherwise.
+			if (!ItemAdded)
 			{
-				MFnTransform TransformNode(obj);
-				MDagPath Path;
-				TransformNode.getPath(Path);
-				LiveLinkStreamManager->AddPropSubject(SubjectFName, Path);
+				if (SelectedRoot.hasFn(MFn::kTransform))
+				{
+					MFnDagNode DagNode(SelectedRoot);
+
+					MDagPath SubjectPath;
+					DagNode.getPath(SubjectPath);
+
+					LiveLinkStreamManager->AddPropSubject(DagNode.name().asChar(), SubjectPath);
+				}
 			}
 		}
-
-		MGlobal::displayInfo(MString("LiveLinkAddSubjectCommand ") + Name);
 		return MS::kSuccess;
 	}
 };
@@ -752,7 +1251,7 @@ public:
 	static void		cleanup() {}
 	static void*	creator() { return new LiveLinkRemoveSubjectCommand(); }
 
-	MStatus			doIt(const MArgList& args)
+	MStatus			doIt(const MArgList& args) override
 	{
 		MSyntax Syntax;
 		Syntax.addArg(MSyntax::kString);
@@ -768,6 +1267,33 @@ public:
 	}
 };
 
+const MString LiveLinkChangeSubjectNameCommandName("LiveLinkChangeSubjectName");
+
+class LiveLinkChangeSubjectNameCommand : public MPxCommand
+{
+public:
+	static void		cleanup() {}
+	static void*	creator() { return new LiveLinkChangeSubjectNameCommand(); }
+
+	MStatus			doIt(const MArgList& args) override
+	{
+		MSyntax Syntax;
+		Syntax.addArg(MSyntax::kString);
+		Syntax.addArg(MSyntax::kString);
+
+		MArgDatabase argData(Syntax, args);
+
+		MString SubjectDagPath;
+		MString NewName;
+		argData.getCommandArgument(0, SubjectDagPath);
+		argData.getCommandArgument(1, NewName);
+
+		LiveLinkStreamManager->ChangeSubjectName(SubjectDagPath, NewName);
+
+		return MS::kSuccess;
+	}
+};
+
 const MString LiveLinkConnectionStatusCommandName("LiveLinkConnectionStatus");
 
 class LiveLinkConnectionStatusCommand : public MPxCommand
@@ -776,7 +1302,7 @@ public:
 	static void		cleanup() {}
 	static void*	creator() { return new LiveLinkConnectionStatusCommand(); }
 
-	MStatus			doIt(const MArgList& args)
+	MStatus			doIt(const MArgList& args) override
 	{
 		MString ConnectionStatus("No Provider (internal error)");
 		bool bConnection = false;
@@ -801,6 +1327,33 @@ public:
 	}
 };
 
+const MString LiveLinkChangeSubjectStreamTypeCommandName("LiveLinkChangeSubjectStreamType");
+
+class LiveLinkChangeSubjectStreamTypeCommand : public MPxCommand
+{
+public:
+	static void		cleanup() {}
+	static void*	creator() { return new LiveLinkChangeSubjectStreamTypeCommand(); }
+
+	MStatus			doIt(const MArgList& args) override
+	{
+		MSyntax Syntax;
+		Syntax.addArg(MSyntax::kString);
+		Syntax.addArg(MSyntax::kString);
+
+		MArgDatabase argData(Syntax, args);
+
+		MString SubjectPath;
+		argData.getCommandArgument(0, SubjectPath);
+		MString StreamType;
+		argData.getCommandArgument(1, StreamType);
+
+		LiveLinkStreamManager->ChangeStreamType(SubjectPath, StreamType);
+
+		return MS::kSuccess;
+	}
+};
+
 void OnForceChange(MTime& time, void* clientData)
 {
 	LiveLinkStreamManager->StreamSubjects();
@@ -990,7 +1543,6 @@ DLLEXPORT MStatus initializePlugin(MObject MayaPluginObject)
 	LiveLinkStreamManager = MakeShareable(new FLiveLinkStreamedSubjectManager());
 
 
-
 	MCallbackId forceUpdateCallbackId = MDGMessage::addForceUpdateCallback((MMessage::MTimeFunction)OnForceChange);
 	myCallbackIds.append(forceUpdateCallbackId);
 
@@ -1007,10 +1559,15 @@ DLLEXPORT MStatus initializePlugin(MObject MayaPluginObject)
 	MCallbackId timerCallback = MTimerMessage::addTimerCallback(5.f, (MMessage::MElapsedTimeFunction)OnInterval);
 	myCallbackIds.append(timerCallback);
 
-	MayaPlugin.registerCommand(LiveLinkSubjectsCommandName, LiveLinkSubjectsCommand::creator);
-	MayaPlugin.registerCommand(LiveLinkAddSubjectCommandName, LiveLinkAddSubjectCommand::creator);
+	MayaPlugin.registerCommand(LiveLinkSubjectNamesCommandName, LiveLinkSubjectNamesCommand::creator);
+	MayaPlugin.registerCommand(LiveLinkSubjectPathsCommandName, LiveLinkSubjectPathsCommand::creator);
+	MayaPlugin.registerCommand(LiveLinkSubjectRolesCommandName, LiveLinkSubjectRolesCommand::creator);
+	MayaPlugin.registerCommand(LiveLinkSubjectTypesCommandName, LiveLinkSubjectTypesCommand::creator);
+	MayaPlugin.registerCommand(LiveLinkAddSelectionCommandName, LiveLinkAddSelectionCommand::creator);
 	MayaPlugin.registerCommand(LiveLinkRemoveSubjectCommandName, LiveLinkRemoveSubjectCommand::creator);
+	MayaPlugin.registerCommand(LiveLinkChangeSubjectNameCommandName, LiveLinkChangeSubjectNameCommand::creator);
 	MayaPlugin.registerCommand(LiveLinkConnectionStatusCommandName, LiveLinkConnectionStatusCommand::creator);
+	MayaPlugin.registerCommand(LiveLinkChangeSubjectStreamTypeCommandName, LiveLinkChangeSubjectStreamTypeCommand::creator);
 
 	// Print to Maya's output window, too!
 	UE_LOG(LogBlankMayaPlugin, Display, TEXT("MayaLiveLinkPlugin initialized"));
@@ -1042,10 +1599,15 @@ DLLEXPORT MStatus uninitializePlugin(MObject MayaPluginObject)
 		MMessage::removeCallbacks(myCallbackIds);
 	}
 
-	MayaPlugin.deregisterCommand(LiveLinkSubjectsCommandName);
-	MayaPlugin.deregisterCommand(LiveLinkAddSubjectCommandName);
+	MayaPlugin.deregisterCommand(LiveLinkSubjectNamesCommandName);
+	MayaPlugin.deregisterCommand(LiveLinkSubjectPathsCommandName);
+	MayaPlugin.deregisterCommand(LiveLinkSubjectRolesCommandName);
+	MayaPlugin.deregisterCommand(LiveLinkSubjectTypesCommandName);
+	MayaPlugin.deregisterCommand(LiveLinkAddSelectionCommandName);
 	MayaPlugin.deregisterCommand(LiveLinkRemoveSubjectCommandName);
+	MayaPlugin.deregisterCommand(LiveLinkChangeSubjectNameCommandName);
 	MayaPlugin.deregisterCommand(LiveLinkConnectionStatusCommandName);
+	MayaPlugin.deregisterCommand(LiveLinkChangeSubjectStreamTypeCommandName);
 
 	if (ConnectionStatusChangedHandle.IsValid())
 	{
diff --git a/Engine/Source/Programs/MayaLiveLinkPlugin/MayaLiveLinkUI.py b/Engine/Source/Programs/MayaLiveLinkPlugin/MayaLiveLinkUI.py
--- a/Engine/Source/Programs/MayaLiveLinkPlugin/MayaLiveLinkUI.py
+++ b/Engine/Source/Programs/MayaLiveLinkPlugin/MayaLiveLinkUI.py
@@ -5,18 +5,58 @@ import maya
 import maya.OpenMaya as OpenMaya
 import maya.OpenMayaMPx as OpenMayaMPx
 import maya.cmds as cmds
+from pymel.core.windows import Callback, CallbackWithArgs
 
-#Populate subjects list from c++ (via command LiveLinkSubjects)
+StreamTypesPerSubjectType = {
+								"Prop": 		["Root Only", "Full Hierarchy"],
+								"Character":	["Root Only", "Full Hierarchy"],
+								"Camera":		["Root Only", "Full Hierarchy", "Camera"],
+								"Light":		["Root Only", "Full Hierarchy", "Light"],
+							}
+
+def OnRemoveSubject(SubjectPath):
+	cmds.LiveLinkRemoveSubject(SubjectPath)
+	RefreshSubjects()
+
+#Populate subjects list from c++
 def PopulateSubjects():
-	Subjects = cmds.LiveLinkSubjects()
-	if Subjects is not None:
-		for Subject in cmds.LiveLinkSubjects():
-			cmds.textScrollList("ActiveSubjects", edit = True, append = Subject)
+	cmds.rowColumnLayout("SubjectLayout", numberOfColumns=5,  adjustableColumn=2, columnWidth=[(1, 20), (2,80), (3, 100), (4, 180), (5, 120)], columnOffset=[(1, 'right', 5), (2, 'right', 10), (4, 'left', 10)], parent="SubjectWrapperLayout")
+	cmds.text(label="")
+	cmds.text(label="Subject Type", font="boldLabelFont", align="left")
+	cmds.text(label="Subject Name", font="boldLabelFont", align="left")
+	cmds.text(label="DAG Path", font="boldLabelFont", align="left")
+	cmds.text(label="Stream Type", font="boldLabelFont", align="left")
+	cmds.rowColumnLayout("SubjectLayout", edit=True, rowOffset=(1, "bottom", 10))
+
+	SubjectNames = cmds.LiveLinkSubjectNames()
+	SubjectPaths = cmds.LiveLinkSubjectPaths()
+	SubjectTypes = cmds.LiveLinkSubjectTypes()
+	SubjectRoles = cmds.LiveLinkSubjectRoles()
+	if SubjectPaths is not None:
+		RowCounter = 0
+		for (SubjectName, SubjectPath, SubjectType, SubjectRole) in zip(SubjectNames, SubjectPaths, SubjectTypes, SubjectRoles):
+			cmds.button(label="-", height=21, statusBarMessage="Remove " + SubjectType, command=Callback(OnRemoveSubject, SubjectPath), parent="SubjectLayout")
+			cmds.text(label=SubjectType, height=21, align="left", parent="SubjectLayout")
+			cmds.textField(text=SubjectName, height=21, changeCommand=CallbackWithArgs(cmds.LiveLinkChangeSubjectName, SubjectPath), parent="SubjectLayout")
+			cmds.text(label=SubjectPath, align="left", height=21, parent="SubjectLayout")
+
+			LayoutName = "ColumnLayoutRow_" + str(RowCounter) # adding a trailing index makes the name unique which is required by the api
+
+			cmds.columnLayout(LayoutName, parent="SubjectLayout")
+			cmds.optionMenu("SubjectTypeMenu", parent=LayoutName, height=21, changeCommand=CallbackWithArgs(cmds.LiveLinkChangeSubjectStreamType, SubjectPath))
+			
+			for StreamType in StreamTypesPerSubjectType[SubjectType]:
+				cmds.menuItem(label=StreamType)
+			
+			StreamTypeIndex = StreamTypesPerSubjectType[SubjectType].index(SubjectRole) + 1 # menu items are 1-indexed
+			cmds.optionMenu("SubjectTypeMenu", edit=True, select=StreamTypeIndex)
+			
+			RowCounter += 1
 
 #Refresh subjects list
 def RefreshSubjects():
 	if (cmds.window(MayaLiveLinkUI.WindowName , exists=True)):
-		cmds.textScrollList("ActiveSubjects", edit=True, removeAll=True)
+		cmds.deleteUI("SubjectLayout")
 		PopulateSubjects()
 
 #Connection UI Colours
@@ -49,11 +89,11 @@ def GetLiveLinkCommandsFromModule(ModuleItems):
 class MayaLiveLinkUI(LiveLinkCommand):
 	WindowName = "MayaLiveLinkUI"
 	Title = "Maya Live Link UI"
-	WindowSize = (500, 350)
+	WindowSize = (500, 300)
 
 	def __init__(self):
 		LiveLinkCommand.__init__(self)
-        
+		
 	# Invoked when the command is run.
 	def doIt(self,argList):
 		if (cmds.window(self.WindowName , exists=True)):
@@ -63,38 +103,31 @@ class MayaLiveLinkUI(LiveLinkCommand):
 		#Get current connection status
 		ConnectionText, ConnectedState = cmds.LiveLinkConnectionStatus()
 		
-		#Type your UI code here
 		cmds.columnLayout( "mainColumn", adjustableColumn=True )
 		cmds.rowLayout("HeaderRow", numberOfColumns=3, adjustableColumn=1, parent = "mainColumn")
 		cmds.text(label="Unreal Engine Live Link", align="left")
 		cmds.text(label="Connection:")
 		cmds.text("ConnectionStatusUI", label=ConnectionText, align="center", backgroundColor=ConnectionColourMap[ConnectedState], width=150)
-		cmds.textScrollList("ActiveSubjects", allowMultiSelection = True, parent = "mainColumn")
+		
+		cmds.separator(h=20, style="none", parent="mainColumn")
+		cmds.columnLayout("SubjectWrapperLayout", parent="mainColumn") # just used as a container that will survive refreshing, so the following controls stay in their correct place
 
 		PopulateSubjects()
 
-		cmds.rowLayout("AddSelectedAsSubject", numberOfColumns=3, parent = "mainColumn")
-		cmds.textField( "NewSubjectName", text = "Maya", parent = "AddSelectedAsSubject", width=300)
-		cmds.button( label='Add Subject', parent = "AddSelectedAsSubject", command=self.AddSubject)
-		cmds.button( label='Remove Subject', parent = "AddSelectedAsSubject", command=self.RemoveSubject)
+		cmds.separator(h=20, style="none", parent="mainColumn")
+		cmds.button( label='Add Selection', statusBarMessage="Adds all selected subjects in the Outliner to LiveLink", parent = "mainColumn", command=self.AddSelection)
 
 		cmds.showWindow( self.WindowName )
 
-	def AddSubject(self, *args):
-		Name = cmds.textField("NewSubjectName", query = True, text = True)
-		cmds.LiveLinkAddSubject(Name)
-		RefreshSubjects()
-
-	def RemoveSubject(self, *args):
-		for ItemToRemove in cmds.textScrollList("ActiveSubjects", q=1, si=1):
-			cmds.LiveLinkRemoveSubject(ItemToRemove)
+	def AddSelection(self, *args):
+		cmds.LiveLinkAddSelection()
 		RefreshSubjects()
 
 # Command to Refresh the subject UI
 class MayaLiveLinkRefreshUI(LiveLinkCommand):
 	def __init__(self):
 		LiveLinkCommand.__init__(self)
-        
+		
 	# Invoked when the command is run.
 	def doIt(self,argList):
 		RefreshSubjects()
@@ -103,7 +136,7 @@ class MayaLiveLinkRefreshUI(LiveLinkCommand):
 class MayaLiveLinkRefreshConnectionUI(LiveLinkCommand):
 	def __init__(self):
 		LiveLinkCommand.__init__(self)
-        
+		
 	# Invoked when the command is run.
 	def doIt(self,argList):
 		if (cmds.window(MayaLiveLinkUI.WindowName , exists=True)):
@@ -114,7 +147,7 @@ class MayaLiveLinkRefreshConnectionUI(LiveLinkCommand):
 class MayaLiveLinkGetActiveCamera(LiveLinkCommand):
 	def __init__(self):
 		LiveLinkCommand.__init__(self)
-        
+		
 	# Invoked when the command is run.
 	def doIt(self,argList):
 		self.clearResult()
