Fix for Maya crash on LiveLink plugin unloading
Fix for Maya LiveLink UI not working for versions of Maya older than 2018

diff --git a/Engine/Source/Programs/MayaLiveLinkPlugin/MayaLiveLinkPlugin.cpp b/Engine/Source/Programs/MayaLiveLinkPlugin/MayaLiveLinkPlugin.cpp
--- a/Engine/Source/Programs/MayaLiveLinkPlugin/MayaLiveLinkPlugin.cpp
+++ b/Engine/Source/Programs/MayaLiveLinkPlugin/MayaLiveLinkPlugin.cpp
@@ -1524,6 +1524,9 @@ DLLEXPORT MStatus initializePlugin(MObject MayaPluginObject)
 		GLog->TearDown(); //clean up existing output devices
 		GLog->AddOutputDevice(new FMayaOutputDevice()); //Add Maya output device
 
+		LiveLinkProvider = ILiveLinkProvider::CreateLiveLinkProvider(TEXT("Maya Live Link"));
+		ConnectionStatusChangedHandle = LiveLinkProvider->RegisterConnStatusChangedHandle(FLiveLinkProviderConnectionStatusChanged::FDelegate::CreateStatic(&OnConnectionStatusChanged));
+
 		bUEInitialized = true; // Dont redo this part if someone unloads and reloads our plugin
 	}
 
@@ -1533,9 +1536,6 @@ DLLEXPORT MStatus initializePlugin(MObject MayaPluginObject)
 		"MayaLiveLinkPlugin",
 		"v1.0");
 
-	LiveLinkProvider = ILiveLinkProvider::CreateLiveLinkProvider(TEXT("Maya Live Link"));
-	ConnectionStatusChangedHandle = LiveLinkProvider->RegisterConnStatusChangedHandle(FLiveLinkProviderConnectionStatusChanged::FDelegate::CreateStatic(&OnConnectionStatusChanged));
-
 	// We do not tick the core engine but we need to tick the ticker to make sure the message bus endpoint in LiveLinkProvider is
 	// up to date
 	FTicker::GetCoreTicker().Tick(1.f);
@@ -1588,17 +1588,14 @@ DLLEXPORT MStatus initializePlugin(MObject MayaPluginObject)
 */
 DLLEXPORT MStatus uninitializePlugin(MObject MayaPluginObject)
 {
+	// Make sure the Garbage Collector does not try to remove Delete Listeners on shutdown as those will be invalid causing a crash
+	GIsRequestingExit = true;
+
 	// Get the plugin API for the plugin object
 	MFnPlugin MayaPlugin(MayaPluginObject);
 
 	// ... do stuff here ...
 
-	if (myCallbackIds.length() != 0)
-	{
-		// Make sure we remove all the callbacks we added
-		MMessage::removeCallbacks(myCallbackIds);
-	}
-
 	MayaPlugin.deregisterCommand(LiveLinkSubjectNamesCommandName);
 	MayaPlugin.deregisterCommand(LiveLinkSubjectPathsCommandName);
 	MayaPlugin.deregisterCommand(LiveLinkSubjectRolesCommandName);
@@ -1609,15 +1606,32 @@ DLLEXPORT MStatus uninitializePlugin(MObject MayaPluginObject)
 	MayaPlugin.deregisterCommand(LiveLinkConnectionStatusCommandName);
 	MayaPlugin.deregisterCommand(LiveLinkChangeSubjectStreamTypeCommandName);
 
+	if (myCallbackIds.length() != 0)
+	{
+		// Make sure we remove all the callbacks we added
+		MMessage::removeCallbacks(myCallbackIds);
+	}
+
 	if (ConnectionStatusChangedHandle.IsValid())
 	{
 		LiveLinkProvider->UnregisterConnStatusChangedHandle(ConnectionStatusChangedHandle);
 		ConnectionStatusChangedHandle.Reset();
 	}
 
-	FTicker::GetCoreTicker().Tick(1.f);
+	// Maya 2016 does not clean up the address space when unloading a plugin.
+	// So if we cleaned up here it would crash in the init above when trying to load the plugin a second.
+	const MString MayaApiVersion = MayaPlugin.apiVersion();
+	if (FString(MayaApiVersion.asChar()).Compare("201700") >= 0)
+	{
+		LiveLinkProvider.Reset();
+		LiveLinkStreamManager.Reset();
 
-	LiveLinkProvider = nullptr;
+		GEngineLoop.AppPreExit();
+		FModuleManager::Get().UnloadModulesAtShutdown();
+		GEngineLoop.AppExit();
+	}
+
+	FTicker::GetCoreTicker().Tick(1.f);
 
 	const MStatus MayaStatusResult = MS::kSuccess;
 	return MayaStatusResult;
diff --git a/Engine/Source/Programs/MayaLiveLinkPlugin/MayaLiveLinkUI.py b/Engine/Source/Programs/MayaLiveLinkPlugin/MayaLiveLinkUI.py
--- a/Engine/Source/Programs/MayaLiveLinkPlugin/MayaLiveLinkUI.py
+++ b/Engine/Source/Programs/MayaLiveLinkPlugin/MayaLiveLinkUI.py
@@ -35,7 +35,7 @@ def PopulateSubjects():
 	if SubjectPaths is not None:
 		RowCounter = 0
 		for (SubjectName, SubjectPath, SubjectType, SubjectRole) in zip(SubjectNames, SubjectPaths, SubjectTypes, SubjectRoles):
-			cmds.button(label="-", height=21, statusBarMessage="Remove " + SubjectType, command=Callback(OnRemoveSubject, SubjectPath), parent="SubjectLayout")
+			cmds.button(label="-", height=21, command=Callback(OnRemoveSubject, SubjectPath), parent="SubjectLayout")
 			cmds.text(label=SubjectType, height=21, align="left", parent="SubjectLayout")
 			cmds.textField(text=SubjectName, height=21, changeCommand=CallbackWithArgs(cmds.LiveLinkChangeSubjectName, SubjectPath), parent="SubjectLayout")
 			cmds.text(label=SubjectPath, align="left", height=21, parent="SubjectLayout")
@@ -115,7 +115,7 @@ class MayaLiveLinkUI(LiveLinkCommand):
 		PopulateSubjects()
 
 		cmds.separator(h=20, style="none", parent="mainColumn")
-		cmds.button( label='Add Selection', statusBarMessage="Adds all selected subjects in the Outliner to LiveLink", parent = "mainColumn", command=self.AddSelection)
+		cmds.button( label='Add Selection', parent = "mainColumn", command=self.AddSelection)
 
 		cmds.showWindow( self.WindowName )
 
